First push of new base material plugin.  Contains a few example material functions/layers for sampling which will be iterated on, as well as a number of coordinate functions.

#rnx

#virtualized

diff --git a/Engine/Plugins/BaseMaterial/BaseMaterial.uplugin b/Engine/Plugins/BaseMaterial/BaseMaterial.uplugin
new file mode 100644
--- /dev/null
+++ b/Engine/Plugins/BaseMaterial/BaseMaterial.uplugin
@@ -0,0 +1,17 @@
+{
+	"FileVersion": 3,
+	"Version": 1,
+	"VersionName": "1.0",
+	"FriendlyName": "BaseMaterial",
+	"Description": "Unified base material function repository",
+	"Category": "Other",
+	"CreatedBy": "Epic Games, Inc.",
+	"CreatedByURL": "https://epicgames.com",
+	"DocsURL": "",
+	"MarketplaceURL": "",
+	"SupportURL": "",
+	"CanContainContent": true,
+	"IsBetaVersion": true,
+	"IsExperimentalVersion": false,
+	"Installed": false
+}
\ No newline at end of file
