Enable BaseMaterial Plugin by default.

#rn Enabled the base material plugin by default which is intended to offer a standard framework for material coordinate and sampling behaviors.

diff --git a/Engine/Plugins/BaseMaterial/BaseMaterial.uplugin b/Engine/Plugins/BaseMaterial/BaseMaterial.uplugin
--- a/Engine/Plugins/BaseMaterial/BaseMaterial.uplugin
+++ b/Engine/Plugins/BaseMaterial/BaseMaterial.uplugin
@@ -13,5 +13,6 @@
 	"CanContainContent": true,
 	"IsBetaVersion": true,
 	"IsExperimentalVersion": false,
-	"Installed": false
+	"Installed": false,
+	"EnabledByDefault" : true
 }
\ No newline at end of file
